Added data interface example implementation

diff --git a/Engine/Plugins/FX/ExampleCustomDataInterface/ExampleCustomDataInterface.uplugin b/Engine/Plugins/FX/ExampleCustomDataInterface/ExampleCustomDataInterface.uplugin
new file mode 100644
--- /dev/null
+++ b/Engine/Plugins/FX/ExampleCustomDataInterface/ExampleCustomDataInterface.uplugin
@@ -0,0 +1,30 @@
+{
+	"FileVersion": 3,
+	"Version": 1,
+	"VersionName": "1.0",
+	"FriendlyName": "Niagara Example Custom DataInterface",
+	"Description": "This plugin contains C++ example content that shows how to write your own data interface for Niagara. \nCheck out the plugin folder Engine/Plugins/FX/ExampleCustomDataInterface for the source files.",
+	"Category": "Other",
+	"CreatedBy": "Epic Games, Inc.",
+	"CreatedByURL": "",
+	"DocsURL": "",
+	"MarketplaceURL": "",
+	"SupportURL": "",
+	"CanContainContent": true,
+	"IsBetaVersion": false,
+	"IsExperimentalVersion": false,
+	"Installed": false,
+	"Modules": [
+		{
+			"Name": "ExampleCustomDataInterface",
+			"Type": "Runtime",
+			"LoadingPhase": "Default"
+		}
+	],
+	"Plugins": [
+		{
+			"Name": "Niagara",
+			"Enabled": true
+		}
+	]
+}
\ No newline at end of file
diff --git a/Engine/Plugins/FX/ExampleCustomDataInterface/Shaders/Private/NiagaraDataInterfaceMousePosition.ush b/Engine/Plugins/FX/ExampleCustomDataInterface/Shaders/Private/NiagaraDataInterfaceMousePosition.ush
new file mode 100644
--- /dev/null
+++ b/Engine/Plugins/FX/ExampleCustomDataInterface/Shaders/Private/NiagaraDataInterfaceMousePosition.ush
@@ -0,0 +1,20 @@
+// Copyright Epic Games, Inc. All Rights Reserved.
+
+
+// when using a template ush file, we need the _{ParameterName} appendix on global functions and parameters, because the template can be included multiple times for different data interfaces in a system.
+float4 MousePosition_{ParameterName};
+
+void GetMousePosition_{ParameterName}(in bool In_Normalized, out float Out_PosX, out float Out_PosY)
+{
+	float4 DIPosition = MousePosition_{ParameterName};
+	if (In_Normalized)
+	{
+		Out_PosX = DIPosition.x / DIPosition.z;
+		Out_PosY = DIPosition.y / DIPosition.w;
+	}
+	else
+	{
+		Out_PosX = DIPosition.x;
+		Out_PosY = DIPosition.y;
+	}
+}
diff --git a/Engine/Plugins/FX/ExampleCustomDataInterface/Source/ExampleCustomDataInterface/ExampleCustomDataInterface.Build.cs b/Engine/Plugins/FX/ExampleCustomDataInterface/Source/ExampleCustomDataInterface/ExampleCustomDataInterface.Build.cs
new file mode 100644
--- /dev/null
+++ b/Engine/Plugins/FX/ExampleCustomDataInterface/Source/ExampleCustomDataInterface/ExampleCustomDataInterface.Build.cs
@@ -0,0 +1,33 @@
+// Copyright Epic Games, Inc. All Rights Reserved.
+
+using UnrealBuildTool;
+
+public class ExampleCustomDataInterface : ModuleRules
+{
+	public ExampleCustomDataInterface(ReadOnlyTargetRules Target) : base(Target)
+	{
+		PCHUsage = PCHUsageMode.UseExplicitOrSharedPCHs;
+		
+		PublicDependencyModuleNames.AddRange(
+			new[]
+			{
+				"Core"
+			}
+			);
+			
+		
+		PrivateDependencyModuleNames.AddRange(
+			new[]
+			{
+				"CoreUObject",
+				"Engine",
+				"Slate",
+				"SlateCore",
+				"Projects",
+				
+				// Data interface dependencies
+				"Niagara", "NiagaraCore", "VectorVM", "RenderCore", "RHI", "UnrealEd"
+			}
+			);
+	}
+}
diff --git a/Engine/Plugins/FX/ExampleCustomDataInterface/Source/ExampleCustomDataInterface/Private/ExampleCustomDataInterface.cpp b/Engine/Plugins/FX/ExampleCustomDataInterface/Source/ExampleCustomDataInterface/Private/ExampleCustomDataInterface.cpp
new file mode 100644
--- /dev/null
+++ b/Engine/Plugins/FX/ExampleCustomDataInterface/Source/ExampleCustomDataInterface/Private/ExampleCustomDataInterface.cpp
@@ -0,0 +1,23 @@
+// Copyright Epic Games, Inc. All Rights Reserved.
+
+#include "ExampleCustomDataInterface.h"
+#include "Misc/Paths.h"
+#include "ShaderCore.h"
+#include "Interfaces/IPluginManager.h"
+
+#define LOCTEXT_NAMESPACE "FExampleCustomDataInterfaceModule"
+
+void FExampleCustomDataInterfaceModule::StartupModule()
+{
+	// map the shader dir so we can use it in the data interface
+	FString PluginShaderDir = FPaths::Combine(IPluginManager::Get().FindPlugin(TEXT("ExampleCustomDataInterface"))->GetBaseDir(), TEXT("Shaders"));
+	AddShaderSourceDirectoryMapping(TEXT("/Plugin/ExampleCustomDataInterface"), PluginShaderDir);
+}
+
+void FExampleCustomDataInterfaceModule::ShutdownModule()
+{
+}
+
+#undef LOCTEXT_NAMESPACE
+	
+IMPLEMENT_MODULE(FExampleCustomDataInterfaceModule, ExampleCustomDataInterface)
\ No newline at end of file
diff --git a/Engine/Plugins/FX/ExampleCustomDataInterface/Source/ExampleCustomDataInterface/Private/NiagaraDataInterfaceMousePosition.cpp b/Engine/Plugins/FX/ExampleCustomDataInterface/Source/ExampleCustomDataInterface/Private/NiagaraDataInterfaceMousePosition.cpp
new file mode 100644
--- /dev/null
+++ b/Engine/Plugins/FX/ExampleCustomDataInterface/Source/ExampleCustomDataInterface/Private/NiagaraDataInterfaceMousePosition.cpp
@@ -0,0 +1,267 @@
+// Copyright Epic Games, Inc. All Rights Reserved.
+
+#include "NiagaraDataInterfaceMousePosition.h"
+
+#include "Editor.h"
+#include "LevelEditorViewport.h"
+#include "NiagaraTypes.h"
+#include "ShaderParameterUtils.h"
+#include "Internationalization/Internationalization.h"
+#include "NiagaraSystemInstance.h"
+#include "NiagaraWorldManager.h"
+#include "ShaderCompilerCore.h"
+#include "GameFramework/PlayerController.h"
+
+#define LOCTEXT_NAMESPACE "NiagaraDataInterfaceMousePosition"
+
+const FName UNiagaraDataInterfaceMousePosition::GetMousePositionName(TEXT("GetMousePosition"));
+static const TCHAR* MouseDITemplateShaderFile = TEXT("/Plugin/ExampleCustomDataInterface/Private/NiagaraDataInterfaceMousePosition.ush");
+
+// the struct used to store our data interface data
+struct FNDIMousePositionInstanceData
+{
+	FVector2f MousePos;
+	FIntPoint ScreenSize;
+};
+
+// this proxy is used to safely copy data between game thread and render thread
+struct FNDIMousePositionProxy : public FNiagaraDataInterfaceProxy
+{
+	virtual int32 PerInstanceDataPassedToRenderThreadSize() const override { return sizeof(FNDIMousePositionInstanceData); }
+
+	static void ProvidePerInstanceDataForRenderThread(void* InDataForRenderThread, void* InDataFromGameThread, const FNiagaraSystemInstanceID& SystemInstance)
+	{
+		const FNDIMousePositionInstanceData* DataFromGameThread = static_cast<FNDIMousePositionInstanceData*>(InDataFromGameThread);
+		FNDIMousePositionInstanceData* DataForRenderThread = static_cast<FNDIMousePositionInstanceData*>(InDataForRenderThread);
+		*DataForRenderThread = *DataFromGameThread;
+	}
+
+	virtual void ConsumePerInstanceDataFromGameThread(void* PerInstanceData, const FNiagaraSystemInstanceID& InstanceID) override
+	{
+		FNDIMousePositionInstanceData* InstanceDataFromGT = static_cast<FNDIMousePositionInstanceData*>(PerInstanceData);
+		FNDIMousePositionInstanceData& InstanceData = SystemInstancesToInstanceData_RT.FindOrAdd(InstanceID);
+		InstanceData = *InstanceDataFromGT;
+	}
+
+	TMap<FNiagaraSystemInstanceID, FNDIMousePositionInstanceData> SystemInstancesToInstanceData_RT;
+};
+
+// creates a new data object to store our position in.
+// Don't keep transient data on the data interface object itself, only use per instance data!
+bool UNiagaraDataInterfaceMousePosition::InitPerInstanceData(void* PerInstanceData,	FNiagaraSystemInstance* SystemInstance)
+{
+	FNDIMousePositionInstanceData* InstanceData = new (PerInstanceData) FNDIMousePositionInstanceData;
+	InstanceData->MousePos = FVector2f::ZeroVector;
+	return true;
+}
+
+// clean up RT instances
+void UNiagaraDataInterfaceMousePosition::DestroyPerInstanceData(void* PerInstanceData, FNiagaraSystemInstance* SystemInstance)
+{
+	FNDIMousePositionInstanceData* InstanceData = static_cast<FNDIMousePositionInstanceData*>(PerInstanceData);
+	InstanceData->~FNDIMousePositionInstanceData();
+
+	ENQUEUE_RENDER_COMMAND(RemoveProxy)
+	(
+		[RT_Proxy=GetProxyAs<FNDIMousePositionProxy>(), InstanceID=SystemInstance->GetId()](FRHICommandListImmediate& CmdList)
+		{
+			RT_Proxy->SystemInstancesToInstanceData_RT.Remove(InstanceID);
+		}
+	);
+}
+
+int32 UNiagaraDataInterfaceMousePosition::PerInstanceDataSize() const
+{
+	return sizeof(FNDIMousePositionInstanceData);
+}
+
+// This ticks on the game thread and lets us do work to initialize the instance data.
+// If you need to do work on the gathered instance data after the simulation is done, use PerInstanceTickPostSimulate() instead. 
+bool UNiagaraDataInterfaceMousePosition::PerInstanceTick(void* PerInstanceData, FNiagaraSystemInstance* SystemInstance,	float DeltaSeconds)
+{
+	check(SystemInstance);
+	FNDIMousePositionInstanceData* InstanceData = static_cast<FNDIMousePositionInstanceData*>(PerInstanceData);
+	if (!InstanceData)
+	{
+		return true;
+	}
+
+	InstanceData->MousePos = FVector2f::ZeroVector;
+
+	// If we have a player controller we use it to capture the mouse position
+	UWorld* World = SystemInstance->GetWorld();
+	if (World && World->GetNumPlayerControllers() > 0)
+	{
+		APlayerController* Controller = World->GetFirstPlayerController();
+		Controller->GetMousePosition(InstanceData->MousePos.X, InstanceData->MousePos.Y);
+		Controller->GetViewportSize(InstanceData->ScreenSize.X, InstanceData->ScreenSize.Y);
+		return false;
+	}
+	
+#if WITH_EDITORONLY_DATA
+	// While in the editor we don't necessarily have a player controller, so we query the viewport object instead 
+	if (GCurrentLevelEditingViewportClient)
+	{
+		InstanceData->MousePos.X = GCurrentLevelEditingViewportClient->Viewport->GetMouseX();
+		InstanceData->MousePos.Y = GCurrentLevelEditingViewportClient->Viewport->GetMouseY();
+		InstanceData->ScreenSize = GCurrentLevelEditingViewportClient->Viewport->GetSizeXY();
+	}
+#endif
+
+	return false;
+}
+
+void UNiagaraDataInterfaceMousePosition::ProvidePerInstanceDataForRenderThread(void* DataForRenderThread, void* PerInstanceData, const FNiagaraSystemInstanceID& SystemInstance)
+{
+	FNDIMousePositionProxy::ProvidePerInstanceDataForRenderThread(DataForRenderThread, PerInstanceData, SystemInstance);
+}
+
+UNiagaraDataInterfaceMousePosition::UNiagaraDataInterfaceMousePosition(FObjectInitializer const& ObjectInitializer)
+	: Super(ObjectInitializer)
+{
+	Proxy.Reset(new FNDIMousePositionProxy());
+}
+
+// this registers our custom DI with Niagara
+void UNiagaraDataInterfaceMousePosition::PostInitProperties()
+{
+	Super::PostInitProperties();
+
+	if (HasAnyFlags(RF_ClassDefaultObject))
+	{
+		ENiagaraTypeRegistryFlags Flags = ENiagaraTypeRegistryFlags::AllowAnyVariable | ENiagaraTypeRegistryFlags::AllowParameter;
+		FNiagaraTypeRegistry::Register(FNiagaraTypeDefinition(GetClass()), Flags);
+	}
+}
+
+// this lists all the functions our DI provides (currently only one)
+void UNiagaraDataInterfaceMousePosition::GetFunctions(TArray<FNiagaraFunctionSignature>& OutFunctions)
+{
+	FNiagaraFunctionSignature Sig;
+	Sig.Name = GetMousePositionName;
+#if WITH_EDITORONLY_DATA
+	Sig.Description = LOCTEXT("GetMousePositionNameFunctionDescription", "Returns the mouse position in screen space.");
+#endif
+	Sig.bMemberFunction = true;
+	Sig.AddInput(FNiagaraVariable(FNiagaraTypeDefinition(GetClass()), TEXT("MousePosition interface")));
+	Sig.AddInput(FNiagaraVariable(FNiagaraTypeDefinition::GetBoolDef(), TEXT("Normalized")));
+	Sig.AddOutput(FNiagaraVariable(FNiagaraTypeDefinition::GetFloatDef(), TEXT("PosX")), LOCTEXT("MousePosXDescription", "Returns the x coordinates in pixels or 0-1 range if normalized"));
+	Sig.AddOutput(FNiagaraVariable(FNiagaraTypeDefinition::GetFloatDef(), TEXT("PosY")), LOCTEXT("MousePosYDescription", "Returns the y coordinates in pixels or 0-1 range if normalized"));
+	OutFunctions.Add(Sig);
+}
+
+// this provides the cpu vm with the correct function to call
+void UNiagaraDataInterfaceMousePosition::GetVMExternalFunction(const FVMExternalFunctionBindingInfo& BindingInfo, void* InstanceData, FVMExternalFunction& OutFunc)
+{
+	if (BindingInfo.Name == GetMousePositionName)
+	{
+		OutFunc = FVMExternalFunction::CreateLambda([this](FVectorVMExternalFunctionContext& Context) { this->GetMousePositionVM(Context); });
+	}
+	else
+	{
+		UE_LOG(LogTemp, Display, TEXT("Could not find data interface external function in %s. Received Name: %s"), *GetPathNameSafe(this), *BindingInfo.Name.ToString());
+	}
+}
+
+// implementation called by the vectorVM
+void UNiagaraDataInterfaceMousePosition::GetMousePositionVM(FVectorVMExternalFunctionContext& Context)
+{
+	VectorVM::FUserPtrHandler<FNDIMousePositionInstanceData> InstData(Context);
+	FNDIInputParam<bool> InNormalized(Context);
+	FNDIOutputParam<float> OutPosX(Context);
+	FNDIOutputParam<float> OutPosY(Context);
+
+	FVector2f MousePos = InstData.Get()->MousePos;
+	FIntPoint ScreenSize = InstData.Get()->ScreenSize;
+
+	// iterate over the particles
+	for (int32 i = 0; i < Context.GetNumInstances(); ++i)
+	{
+		if (InNormalized.GetAndAdvance())
+		{
+			OutPosX.SetAndAdvance(MousePos.X / ScreenSize.X);
+			OutPosY.SetAndAdvance(MousePos.Y / ScreenSize.Y);
+		}
+		else
+		{
+			OutPosX.SetAndAdvance(MousePos.X);
+			OutPosY.SetAndAdvance(MousePos.Y);
+		}
+	}
+}
+
+#if WITH_EDITORONLY_DATA
+
+// this lets the niagara compiler know that it needs to recompile an effect when our hlsl file changes
+bool UNiagaraDataInterfaceMousePosition::AppendCompileHash(FNiagaraCompileHashVisitor* InVisitor) const
+{
+	if (!Super::AppendCompileHash(InVisitor))
+	{
+		return false;
+	}
+
+	FSHAHash Hash = GetShaderFileHash(MouseDITemplateShaderFile, SP_PCD3D_SM5);
+	InVisitor->UpdateString(TEXT("NiagaraDataInterfaceMousePositionHLSLSource"), Hash.ToString());
+	return true;
+}
+
+// This can be used to provide the hlsl code for gpu scripts. If the DI supports only cpu implementations, this is not needed.
+// We don't need to actually print our function code to OutHLSL here because we use a template file that gets appended in GetParameterDefinitionHLSL().
+// If the hlsl function is so simple that it does not need bound shader parameters, then this method can be used instead of GetParameterDefinitionHLSL.
+bool UNiagaraDataInterfaceMousePosition::GetFunctionHLSL(const FNiagaraDataInterfaceGPUParamInfo& ParamInfo, const FNiagaraDataInterfaceGeneratedFunction& FunctionInfo, int FunctionInstanceIndex, FString& OutHLSL)
+{
+	return FunctionInfo.DefinitionName == GetMousePositionName;
+}
+
+// this loads our hlsl template script file and 
+void UNiagaraDataInterfaceMousePosition::GetParameterDefinitionHLSL(const FNiagaraDataInterfaceGPUParamInfo& ParamInfo,	FString& OutHLSL)
+{
+	TMap<FString, FStringFormatArg> TemplateArgs =
+	{
+		{TEXT("ParameterName"),	ParamInfo.DataInterfaceHLSLSymbol},
+	};
+
+	FString TemplateFile;
+	LoadShaderSourceFile(MouseDITemplateShaderFile, SP_PCD3D_SM5, &TemplateFile, nullptr);
+	OutHLSL += FString::Format(*TemplateFile, TemplateArgs);
+}
+
+#endif
+
+//////////////////////////////////////////////////////////////////////////
+// Compute Shader Binding
+struct FNiagaraDataInterfaceParametersCS_MousePosition : public FNiagaraDataInterfaceParametersCS
+{
+	DECLARE_TYPE_LAYOUT(FNiagaraDataInterfaceParametersCS_MousePosition, NonVirtual);
+
+	void Bind(const FNiagaraDataInterfaceGPUParamInfo& ParameterInfo, const class FShaderParameterMap& ParameterMap)
+	{
+		MousePositionParam.Bind(ParameterMap, *(MousePositionParamName + ParameterInfo.DataInterfaceHLSLSymbol));
+	}
+
+	void Set(FRHICommandList& RHICmdList, const FNiagaraDataInterfaceSetArgs& Context) const
+	{
+		check(IsInRenderingThread());
+		
+		FRHIComputeShader* ComputeShaderRHI = Context.Shader.GetComputeShader();
+		FNDIMousePositionProxy* DataInterfaceProxy = static_cast<FNDIMousePositionProxy*>(Context.DataInterface);
+		FNDIMousePositionInstanceData* InstanceData = DataInterfaceProxy->SystemInstancesToInstanceData_RT.Find(Context.SystemInstanceID);
+		check(InstanceData != nullptr);
+
+		FVector4f ShaderData(InstanceData->MousePos.X, InstanceData->MousePos.Y, InstanceData->ScreenSize.X, InstanceData->ScreenSize.Y);
+		SetShaderValue(RHICmdList, ComputeShaderRHI, MousePositionParam, ShaderData);
+	}
+
+private:
+	// Mouse position from game thread
+	LAYOUT_FIELD(FShaderParameter, MousePositionParam);
+
+	static const FString MousePositionParamName;
+};
+
+const FString FNiagaraDataInterfaceParametersCS_MousePosition::MousePositionParamName(TEXT("MousePosition_"));
+
+IMPLEMENT_TYPE_LAYOUT(FNiagaraDataInterfaceParametersCS_MousePosition);
+IMPLEMENT_NIAGARA_DI_PARAMETER(UNiagaraDataInterfaceMousePosition, FNiagaraDataInterfaceParametersCS_MousePosition);
+
+#undef LOCTEXT_NAMESPACE
diff --git a/Engine/Plugins/FX/ExampleCustomDataInterface/Source/ExampleCustomDataInterface/Private/NiagaraDataInterfaceMousePosition.h b/Engine/Plugins/FX/ExampleCustomDataInterface/Source/ExampleCustomDataInterface/Private/NiagaraDataInterfaceMousePosition.h
new file mode 100644
--- /dev/null
+++ b/Engine/Plugins/FX/ExampleCustomDataInterface/Source/ExampleCustomDataInterface/Private/NiagaraDataInterfaceMousePosition.h
@@ -0,0 +1,43 @@
+// Copyright Epic Games, Inc. All Rights Reserved.
+#pragma once
+
+#include "NiagaraCommon.h"
+#include "NiagaraShared.h"
+#include "VectorVM.h"
+#include "NiagaraDataInterface.h"
+#include "NiagaraDataInterfaceMousePosition.generated.h"
+
+UCLASS(EditInlineNew, Category = "Mouse", meta = (DisplayName = "MousePosition Query"))
+class UNiagaraDataInterfaceMousePosition : public UNiagaraDataInterface
+{
+	GENERATED_UCLASS_BODY()
+
+public:
+	// needed to set up gpu script parameters
+	DECLARE_NIAGARA_DI_PARAMETER();
+	
+	//UObject Interface
+	virtual void PostInitProperties() override;
+	//UObject Interface End
+
+	//UNiagaraDataInterface Interface
+	virtual void GetFunctions(TArray<FNiagaraFunctionSignature>& OutFunctions) override;
+	virtual void GetVMExternalFunction(const FVMExternalFunctionBindingInfo& BindingInfo, void* InstanceData, FVMExternalFunction &OutFunc) override;
+	virtual bool CanExecuteOnTarget(ENiagaraSimTarget Target) const override { return true; }
+#if WITH_EDITORONLY_DATA
+	virtual bool AppendCompileHash(FNiagaraCompileHashVisitor* InVisitor) const override;
+	virtual bool GetFunctionHLSL(const FNiagaraDataInterfaceGPUParamInfo& ParamInfo, const FNiagaraDataInterfaceGeneratedFunction& FunctionInfo, int FunctionInstanceIndex, FString& OutHLSL) override;
+	virtual void GetParameterDefinitionHLSL(const FNiagaraDataInterfaceGPUParamInfo& ParamInfo, FString& OutHLSL) override;
+#endif
+	virtual bool InitPerInstanceData(void* PerInstanceData, FNiagaraSystemInstance* SystemInstance) override;
+	virtual void DestroyPerInstanceData(void* PerInstanceData, FNiagaraSystemInstance* SystemInstance) override;
+	virtual int32 PerInstanceDataSize() const override;
+	virtual bool HasPreSimulateTick() const override { return true; }
+	virtual bool PerInstanceTick(void* PerInstanceData, FNiagaraSystemInstance* SystemInstance, float DeltaSeconds) override;
+	virtual void ProvidePerInstanceDataForRenderThread(void* DataForRenderThread, void* PerInstanceData, const FNiagaraSystemInstanceID& SystemInstance) override;
+	//UNiagaraDataInterface Interface
+
+	void GetMousePositionVM(FVectorVMExternalFunctionContext& Context);	
+private:
+	static const FName GetMousePositionName;
+};
diff --git a/Engine/Plugins/FX/ExampleCustomDataInterface/Source/ExampleCustomDataInterface/Public/ExampleCustomDataInterface.h b/Engine/Plugins/FX/ExampleCustomDataInterface/Source/ExampleCustomDataInterface/Public/ExampleCustomDataInterface.h
new file mode 100644
--- /dev/null
+++ b/Engine/Plugins/FX/ExampleCustomDataInterface/Source/ExampleCustomDataInterface/Public/ExampleCustomDataInterface.h
@@ -0,0 +1,15 @@
+// Copyright Epic Games, Inc. All Rights Reserved.
+
+#pragma once
+
+#include "CoreMinimal.h"
+#include "Modules/ModuleManager.h"
+
+class FExampleCustomDataInterfaceModule : public IModuleInterface
+{
+public:
+
+	/** IModuleInterface implementation */
+	virtual void StartupModule() override;
+	virtual void ShutdownModule() override;
+};
