Build fix for non-editor mode

diff --git a/Engine/Plugins/FX/ExampleCustomDataInterface/Source/ExampleCustomDataInterface/ExampleCustomDataInterface.Build.cs b/Engine/Plugins/FX/ExampleCustomDataInterface/Source/ExampleCustomDataInterface/ExampleCustomDataInterface.Build.cs
--- a/Engine/Plugins/FX/ExampleCustomDataInterface/Source/ExampleCustomDataInterface/ExampleCustomDataInterface.Build.cs
+++ b/Engine/Plugins/FX/ExampleCustomDataInterface/Source/ExampleCustomDataInterface/ExampleCustomDataInterface.Build.cs
@@ -12,8 +12,7 @@ public class ExampleCustomDataInterface : ModuleRules
 			new[]
 			{
 				"Core"
-			}
-			);
+			});
 			
 		
 		PrivateDependencyModuleNames.AddRange(
@@ -26,8 +25,12 @@ public class ExampleCustomDataInterface : ModuleRules
 				"Projects",
 				
 				// Data interface dependencies
-				"Niagara", "NiagaraCore", "VectorVM", "RenderCore", "RHI", "UnrealEd"
-			}
-			);
+				"Niagara", "NiagaraCore", "VectorVM", "RenderCore", "RHI"
+			});
+		
+		if (Target.bBuildEditor == true)
+		{
+			PrivateDependencyModuleNames.Add("UnrealEd");
+		}
 	}
 }
diff --git a/Engine/Plugins/FX/ExampleCustomDataInterface/Source/ExampleCustomDataInterface/Private/NiagaraDataInterfaceMousePosition.cpp b/Engine/Plugins/FX/ExampleCustomDataInterface/Source/ExampleCustomDataInterface/Private/NiagaraDataInterfaceMousePosition.cpp
--- a/Engine/Plugins/FX/ExampleCustomDataInterface/Source/ExampleCustomDataInterface/Private/NiagaraDataInterfaceMousePosition.cpp
+++ b/Engine/Plugins/FX/ExampleCustomDataInterface/Source/ExampleCustomDataInterface/Private/NiagaraDataInterfaceMousePosition.cpp
@@ -2,8 +2,11 @@
 
 #include "NiagaraDataInterfaceMousePosition.h"
 
+#if WITH_EDITORONLY_DATA
 #include "Editor.h"
 #include "LevelEditorViewport.h"
+#endif
+
 #include "NiagaraTypes.h"
 #include "ShaderParameterUtils.h"
 #include "Internationalization/Internationalization.h"
