Example DI source improvements

diff --git a/Engine/Plugins/FX/ExampleCustomDataInterface/Source/ExampleCustomDataInterface/Private/NiagaraDataInterfaceMousePosition.cpp b/Engine/Plugins/FX/ExampleCustomDataInterface/Source/ExampleCustomDataInterface/Private/NiagaraDataInterfaceMousePosition.cpp
--- a/Engine/Plugins/FX/ExampleCustomDataInterface/Source/ExampleCustomDataInterface/Private/NiagaraDataInterfaceMousePosition.cpp
+++ b/Engine/Plugins/FX/ExampleCustomDataInterface/Source/ExampleCustomDataInterface/Private/NiagaraDataInterfaceMousePosition.cpp
@@ -34,8 +34,11 @@ struct FNDIMousePositionProxy : public FNiagaraDataInterfaceProxy
 
 	static void ProvidePerInstanceDataForRenderThread(void* InDataForRenderThread, void* InDataFromGameThread, const FNiagaraSystemInstanceID& SystemInstance)
 	{
+		// initialize the render thread instance data into the pre-allocated memory
+		FNDIMousePositionInstanceData* DataForRenderThread = new (InDataForRenderThread) FNDIMousePositionInstanceData();
+
+		// we're just copying the game thread data, but the render thread data can be initialized to anything here and can be another struct entirely
 		const FNDIMousePositionInstanceData* DataFromGameThread = static_cast<FNDIMousePositionInstanceData*>(InDataFromGameThread);
-		FNDIMousePositionInstanceData* DataForRenderThread = static_cast<FNDIMousePositionInstanceData*>(InDataForRenderThread);
 		*DataForRenderThread = *DataFromGameThread;
 	}
 
@@ -44,6 +47,9 @@ struct FNDIMousePositionProxy : public FNiagaraDataInterfaceProxy
 		FNDIMousePositionInstanceData* InstanceDataFromGT = static_cast<FNDIMousePositionInstanceData*>(PerInstanceData);
 		FNDIMousePositionInstanceData& InstanceData = SystemInstancesToInstanceData_RT.FindOrAdd(InstanceID);
 		InstanceData = *InstanceDataFromGT;
+
+		// we call the destructor here to clean up the GT data. Without this we could be leaking memory.
+		InstanceDataFromGT->~FNDIMousePositionInstanceData();
 	}
 
 	TMap<FNiagaraSystemInstanceID, FNDIMousePositionInstanceData> SystemInstancesToInstanceData_RT;
