- Convert example mouse data interface to parameter bindings

diff --git a/Engine/Plugins/FX/ExampleCustomDataInterface/Shaders/Private/NiagaraDataInterfaceMousePosition.ush b/Engine/Plugins/FX/ExampleCustomDataInterface/Shaders/Private/NiagaraDataInterfaceMousePosition.ush
--- a/Engine/Plugins/FX/ExampleCustomDataInterface/Shaders/Private/NiagaraDataInterfaceMousePosition.ush
+++ b/Engine/Plugins/FX/ExampleCustomDataInterface/Shaders/Private/NiagaraDataInterfaceMousePosition.ush
@@ -2,11 +2,11 @@
 
 
 // when using a template ush file, we need the _{ParameterName} appendix on global functions and parameters, because the template can be included multiple times for different data interfaces in a system.
-float4 MousePosition_{ParameterName};
+float4 {ParameterName}_MousePosition;
 
 void GetMousePosition_{ParameterName}(in bool In_Normalized, out float Out_PosX, out float Out_PosY)
 {
-	float4 DIPosition = MousePosition_{ParameterName};
+	float4 DIPosition = {ParameterName}_MousePosition;
 	if (In_Normalized)
 	{
 		Out_PosX = DIPosition.x / DIPosition.z;
diff --git a/Engine/Plugins/FX/ExampleCustomDataInterface/Source/ExampleCustomDataInterface/Private/NiagaraDataInterfaceMousePosition.cpp b/Engine/Plugins/FX/ExampleCustomDataInterface/Source/ExampleCustomDataInterface/Private/NiagaraDataInterfaceMousePosition.cpp
--- a/Engine/Plugins/FX/ExampleCustomDataInterface/Source/ExampleCustomDataInterface/Private/NiagaraDataInterfaceMousePosition.cpp
+++ b/Engine/Plugins/FX/ExampleCustomDataInterface/Source/ExampleCustomDataInterface/Private/NiagaraDataInterfaceMousePosition.cpp
@@ -10,6 +10,7 @@
 #include "NiagaraTypes.h"
 #include "ShaderParameterUtils.h"
 #include "Internationalization/Internationalization.h"
+#include "NiagaraShaderParametersBuilder.h"
 #include "NiagaraSystemInstance.h"
 #include "NiagaraWorldManager.h"
 #include "ShaderCompilerCore.h"
@@ -211,6 +212,7 @@ bool UNiagaraDataInterfaceMousePosition::AppendCompileHash(FNiagaraCompileHashVi
 
 	FSHAHash Hash = GetShaderFileHash(MouseDITemplateShaderFile, SP_PCD3D_SM5);
 	InVisitor->UpdateString(TEXT("NiagaraDataInterfaceMousePositionHLSLSource"), Hash.ToString());
+	InVisitor->UpdateShaderParameters<FShaderParameters>();
 	return true;
 }
 
@@ -237,40 +239,23 @@ void UNiagaraDataInterfaceMousePosition::GetParameterDefinitionHLSL(const FNiaga
 
 #endif
 
-//////////////////////////////////////////////////////////////////////////
-// Compute Shader Binding
-struct FNiagaraDataInterfaceParametersCS_MousePosition : public FNiagaraDataInterfaceParametersCS
+// This fills in the expected parameter bindings we use to send data to the GPU
+void UNiagaraDataInterfaceMousePosition::BuildShaderParameters(FNiagaraShaderParametersBuilder& ShaderParametersBuilder) const
 {
-	DECLARE_TYPE_LAYOUT(FNiagaraDataInterfaceParametersCS_MousePosition, NonVirtual);
-
-	void Bind(const FNiagaraDataInterfaceGPUParamInfo& ParameterInfo, const class FShaderParameterMap& ParameterMap)
-	{
-		MousePositionParam.Bind(ParameterMap, *(MousePositionParamName + ParameterInfo.DataInterfaceHLSLSymbol));
-	}
-
-	void Set(FRHICommandList& RHICmdList, const FNiagaraDataInterfaceSetArgs& Context) const
-	{
-		check(IsInRenderingThread());
-		
-		FRHIComputeShader* ComputeShaderRHI = Context.Shader.GetComputeShader();
-		FNDIMousePositionProxy* DataInterfaceProxy = static_cast<FNDIMousePositionProxy*>(Context.DataInterface);
-		FNDIMousePositionInstanceData* InstanceData = DataInterfaceProxy->SystemInstancesToInstanceData_RT.Find(Context.SystemInstanceID);
-		check(InstanceData != nullptr);
-
-		FVector4f ShaderData(InstanceData->MousePos.X, InstanceData->MousePos.Y, InstanceData->ScreenSize.X, InstanceData->ScreenSize.Y);
-		SetShaderValue(RHICmdList, ComputeShaderRHI, MousePositionParam, ShaderData);
-	}
-
-private:
-	// Mouse position from game thread
-	LAYOUT_FIELD(FShaderParameter, MousePositionParam);
-
-	static const FString MousePositionParamName;
-};
-
-const FString FNiagaraDataInterfaceParametersCS_MousePosition::MousePositionParamName(TEXT("MousePosition_"));
+	ShaderParametersBuilder.AddNestedStruct<FShaderParameters>();
+}
 
-IMPLEMENT_TYPE_LAYOUT(FNiagaraDataInterfaceParametersCS_MousePosition);
-IMPLEMENT_NIAGARA_DI_PARAMETER(UNiagaraDataInterfaceMousePosition, FNiagaraDataInterfaceParametersCS_MousePosition);
+// This fills in the parameters to send to the GPU
+void UNiagaraDataInterfaceMousePosition::SetShaderParameters(const FNiagaraDataInterfaceSetShaderParametersContext& Context) const
+{
+	FNDIMousePositionProxy& DataInterfaceProxy = Context.GetProxy<FNDIMousePositionProxy>();
+	FNDIMousePositionInstanceData& InstanceData = DataInterfaceProxy.SystemInstancesToInstanceData_RT.FindChecked(Context.GetSystemInstanceID());
+
+	FShaderParameters* ShaderParameters = Context.GetParameterNestedStruct<FShaderParameters>();
+	ShaderParameters->MousePosition.X = InstanceData.MousePos.X;
+	ShaderParameters->MousePosition.Y = InstanceData.MousePos.Y;
+	ShaderParameters->MousePosition.Z = InstanceData.ScreenSize.X;
+	ShaderParameters->MousePosition.W = InstanceData.ScreenSize.Y;
+}
 
 #undef LOCTEXT_NAMESPACE
diff --git a/Engine/Plugins/FX/ExampleCustomDataInterface/Source/ExampleCustomDataInterface/Private/NiagaraDataInterfaceMousePosition.h b/Engine/Plugins/FX/ExampleCustomDataInterface/Source/ExampleCustomDataInterface/Private/NiagaraDataInterfaceMousePosition.h
--- a/Engine/Plugins/FX/ExampleCustomDataInterface/Source/ExampleCustomDataInterface/Private/NiagaraDataInterfaceMousePosition.h
+++ b/Engine/Plugins/FX/ExampleCustomDataInterface/Source/ExampleCustomDataInterface/Private/NiagaraDataInterfaceMousePosition.h
@@ -12,10 +12,11 @@ class UNiagaraDataInterfaceMousePosition : public UNiagaraDataInterface
 {
 	GENERATED_UCLASS_BODY()
 
+	BEGIN_SHADER_PARAMETER_STRUCT(FShaderParameters, )
+		SHADER_PARAMETER(FVector4f,				MousePosition)
+	END_SHADER_PARAMETER_STRUCT()
+
 public:
-	// needed to set up gpu script parameters
-	DECLARE_NIAGARA_DI_PARAMETER();
-	
 	//UObject Interface
 	virtual void PostInitProperties() override;
 	//UObject Interface End
@@ -29,6 +30,10 @@ public:
 	virtual bool GetFunctionHLSL(const FNiagaraDataInterfaceGPUParamInfo& ParamInfo, const FNiagaraDataInterfaceGeneratedFunction& FunctionInfo, int FunctionInstanceIndex, FString& OutHLSL) override;
 	virtual void GetParameterDefinitionHLSL(const FNiagaraDataInterfaceGPUParamInfo& ParamInfo, FString& OutHLSL) override;
 #endif
+	virtual bool UseLegacyShaderBindings() const override { return false; }
+	virtual void BuildShaderParameters(FNiagaraShaderParametersBuilder& ShaderParametersBuilder) const override;
+	virtual void SetShaderParameters(const FNiagaraDataInterfaceSetShaderParametersContext& Context) const override;
+
 	virtual bool InitPerInstanceData(void* PerInstanceData, FNiagaraSystemInstance* SystemInstance) override;
 	virtual void DestroyPerInstanceData(void* PerInstanceData, FNiagaraSystemInstance* SystemInstance) override;
 	virtual int32 PerInstanceDataSize() const override;
